Restore October [REDACTED_PHONE]) as the preferred GDK version

#rnx

#Synced-CL 49766262

diff --git a/Engine/Config/Microsoft/GDK_SDK.json b/Engine/Config/Microsoft/GDK_SDK.json
--- a/Engine/Config/Microsoft/GDK_SDK.json
+++ b/Engine/Config/Microsoft/GDK_SDK.json
@@ -1,7 +1,7 @@
 {
-	"MainGDKVersion": "240605",
-	"MinGDKVersion": "240604",
-	"MaxGDKVersion": "240605",
+	"MainGDKVersion": "251000",
+	"MinGDKVersion": "240605",
+	"MaxGDKVersion": "251099",
 
 	"//GRDK_1": "240305 has a byte-ordering problem with PlayFab that causes compile errors",
 	"BannedGDKVersions": [
