set the minimum GDK version to 250404

#rnx

#Synced-CL 50541135

diff --git a/Engine/Config/Microsoft/GDK_SDK.json b/Engine/Config/Microsoft/GDK_SDK.json
--- a/Engine/Config/Microsoft/GDK_SDK.json
+++ b/Engine/Config/Microsoft/GDK_SDK.json
@@ -1,6 +1,6 @@
 {
 	"MainGDKVersion": "251000",
-	"MinGDKVersion": "240605",
+	"MinGDKVersion": "250404",
 	"MaxGDKVersion": "251099",
 
 	"//GRDK_1": "240305 has a byte-ordering problem with PlayFab that causes compile errors",
